In-Editor UI for Turnkey.

#rnx

diff --git a/Engine/Source/Developer/LauncherServices/LauncherServices.Build.cs b/Engine/Source/Developer/LauncherServices/LauncherServices.Build.cs
--- a/Engine/Source/Developer/LauncherServices/LauncherServices.Build.cs
+++ b/Engine/Source/Developer/LauncherServices/LauncherServices.Build.cs
@@ -25,6 +25,7 @@ namespace UnrealBuildTool.Rules
 					"TargetPlatform",
 					"UnrealEdMessages",
 					"JSON",
+					"TurnkeyIO",
 				});
 
 			PrivateIncludePathModuleNames.AddRange(
diff --git a/Engine/Source/Developer/LauncherServices/Private/Launcher/LauncherWorker.cpp b/Engine/Source/Developer/LauncherServices/Private/Launcher/LauncherWorker.cpp
--- a/Engine/Source/Developer/LauncherServices/Private/Launcher/LauncherWorker.cpp
+++ b/Engine/Source/Developer/LauncherServices/Private/Launcher/LauncherWorker.cpp
@@ -7,6 +7,7 @@
 #include "ISourceCodeAccessModule.h"
 #include "ITargetDeviceProxy.h"
 #include "ITargetDeviceProxyManager.h"
+#include "ITurnkeyIOModule.h"
 #include "Misc/CommandLine.h"
 #include "Misc/Paths.h"
 #include "HAL/ThreadSafeCounter.h"
@@ -979,7 +980,7 @@ void FLauncherWorker::CreateAndExecuteTasks( const ILauncherProfileRef& InProfil
 	FString TurnkeyCommand;
 	if (Profile->ShouldUpdateDeviceFlash() && DeviceGroup->GetDeviceIDs().Num() >= 0)
 	{
-		TurnkeyCommand = FString::Printf(TEXT("Turnkey -command=VerifySdk -type=Flash -device=%s -UpdateIfNeeded -EditorIO -utf8output -WaitForUATMutex"), *FString::Join(DeviceGroup->GetDeviceIDs(), TEXT("+")));
+		TurnkeyCommand = FString::Printf(TEXT("Turnkey -command=VerifySdk -type=Flash -device=%s -UpdateIfNeeded -utf8output -WaitForUATMutex %s"), *FString::Join(DeviceGroup->GetDeviceIDs(), TEXT("+")), *ITurnkeyIOModule::Get().GetUATParams());
 	}
 
 	TSharedPtr<FLauncherTask> BuildTask = MakeShareable(new FLauncherUATTask(UATCommand, TEXT("Build Task"), TEXT("Launching UAT..."), ReadPipe, WritePipe, InProfile->GetEditorExe(), ProcHandle, this, StartString, TurnkeyCommand));
diff --git a/Engine/Source/Developer/TurnkeyIO/Private/STurnkeyIOWidgets.cpp b/Engine/Source/Developer/TurnkeyIO/Private/STurnkeyIOWidgets.cpp
new file mode 100644
--- /dev/null
+++ b/Engine/Source/Developer/TurnkeyIO/Private/STurnkeyIOWidgets.cpp
@@ -0,0 +1,261 @@
+// Copyright Epic Games, Inc. All Rights Reserved.
+
+#include "STurnkeyIOWidgets.h"
+
+#if WITH_TURNKEY_EDITOR_IO_SERVER
+#include "Widgets/Text/STextBlock.h"
+#include "Widgets/Input/SButton.h"
+#include "Widgets/Input/SCheckBox.h"
+#include "Widgets/Input/SEditableTextBox.h"
+#include "Dom/JsonObject.h"
+
+#define LOCTEXT_NAMESPACE "TurnkeySlateIO"
+
+
+
+void STurnkeyReadInputModal::Construct( const FArguments& InArgs )
+{
+	// store arguments
+	Prompt = InArgs._Prompt;
+	Default = InArgs._Default;
+	OnFinished = InArgs._OnFinished;
+	Value = Default;
+
+	// construct widget
+	this->ChildSlot
+	[
+		SNew(SBorder)
+		.BorderImage(FAppStyle::Get().GetBrush("Brushes.Panel"))
+		.Padding(8)
+		[
+			SNew(SVerticalBox)
+
+			// prompt 
+			+SVerticalBox::Slot()
+			.AutoHeight()
+			.Padding(4,4)
+			[
+				SNew(STextBlock)
+				.Text( FText::FromString(Prompt) )
+			]
+
+			// text input
+			+SVerticalBox::Slot()
+			.FillHeight(1)
+			.Padding(4,4)
+			.VAlign(VAlign_Top)
+			[
+				SNew(SEditableTextBox)
+				.Text_Lambda( [=]() { return FText::FromString(Value); } )
+				.OnTextChanged_Lambda( [=]( const FText& NewText ) { Value = NewText.ToString(); } )
+				.OnTextCommitted_Lambda( [=]( const FText& NewText, ETextCommit::Type CommitType )
+				{
+					if (CommitType == ETextCommit::OnEnter)
+					{
+						FinishAction();
+					}
+				})
+			]
+
+			// button area
+			+SVerticalBox::Slot()
+			.AutoHeight()
+			.HAlign(HAlign_Right)
+			.VAlign(VAlign_Bottom)
+			[
+				// ok button
+				SNew(SHorizontalBox)
+				+SHorizontalBox::Slot()
+				.Padding(5,5)
+				.AutoWidth()
+				[
+					SNew(SButton)
+					.ButtonStyle(FAppStyle::Get(), "PrimaryButton")
+					.TextStyle( FAppStyle::Get(), "DialogButtonText" )
+					.HAlign(HAlign_Center)
+					.VAlign(VAlign_Center)
+					.Text(LOCTEXT("OK", "OK"))
+					.OnClicked_Lambda( [=]()
+					{ 
+						FinishAction(); 
+						return FReply::Handled();
+					})
+				]
+			]
+		]
+	];
+}
+
+
+void STurnkeyReadInputModal::FinishAction()
+{
+	TSharedPtr<FJsonObject> Result = MakeShared<FJsonObject>();
+	Result->SetStringField( TEXT("Result"), Value );
+
+	OnFinished.ExecuteIfBound(Result);
+}
+
+
+
+
+
+
+void STurnkeyReadInputIntModal::Construct( const FArguments& InArgs )
+{
+	// store arguments
+	Prompt = InArgs._Prompt;
+	Options = InArgs._Options;
+	bIsCancellable = InArgs._IsCancellable;
+	DefaultValue = InArgs._DefaultValue;
+	OnFinished = InArgs._OnFinished;
+	Value = DefaultValue;
+
+	// construct widget
+	TSharedPtr<SVerticalBox> VerticalBox;
+	TSharedPtr<SHorizontalBox> ButtonBox;
+	this->ChildSlot
+	[
+		SNew(SBorder)
+		.BorderImage(FAppStyle::Get().GetBrush("Brushes.Panel"))
+		.Padding(8)
+		[
+			SNew(SVerticalBox)
+
+			// prompt
+			+SVerticalBox::Slot()
+			.AutoHeight()
+			.Padding(4,4)
+			[
+				SNew(STextBlock)
+				.Text( FText::FromString(Prompt) )
+			]
+
+			// container widget 
+			+SVerticalBox::Slot()
+			.FillHeight(1)
+			.Padding(16,16)
+			[
+				SAssignNew(VerticalBox, SVerticalBox)
+			]
+
+			// button area
+			+SVerticalBox::Slot()
+			.AutoHeight()
+			.HAlign(HAlign_Right)
+			.VAlign(VAlign_Bottom)
+			[
+				SAssignNew(ButtonBox, SHorizontalBox)
+
+				// ok button
+				+SHorizontalBox::Slot()
+				.AutoWidth()
+				.Padding(5,5)
+				[
+					SNew(SButton)
+					.ButtonStyle(FAppStyle::Get(), "PrimaryButton")
+					.TextStyle( FAppStyle::Get(), "DialogButtonText" )
+					.HAlign(HAlign_Center)
+					.VAlign(VAlign_Center)
+					.Text(LOCTEXT("OK", "OK"))
+					.IsEnabled_Lambda( [=]() { return Value > 0; } )
+					.OnClicked_Lambda( [=]()
+					{ 
+						FinishAction();
+						return FReply::Handled();
+					})
+				]
+			]
+		]
+	];
+
+	// add cancel button, if permitted
+	if ( bIsCancellable )
+	{
+		ButtonBox->AddSlot()
+		.AutoWidth()
+		.Padding(5,5)
+		[
+			SNew(SButton)
+			.ButtonStyle(FAppStyle::Get(), "Button")
+			.TextStyle( FAppStyle::Get(), "DialogButtonText" )
+			.HAlign(HAlign_Center)
+			.VAlign(VAlign_Center)
+			.Text(LOCTEXT("Cancel", "Cancel"))
+			.OnClicked_Lambda( [=]()
+			{ 
+				Value = 0;
+				FinishAction();
+				return FReply::Handled();
+			})
+		];
+	}
+
+	// add all options
+	int ResultValue = 1;
+	for( int OptionIndex = 0; OptionIndex < Options.Num(); OptionIndex++ )
+	{
+		FString OptionValue = Options[OptionIndex];
+		if( OptionValue.RemoveFromStart(TEXT(";")) )
+		{
+			// just a comment - not a selectable option. create a heading
+			VerticalBox->AddSlot()
+			.Padding(4,4)
+			[
+				SNew(STextBlock)
+				.Text( FText::FromString(OptionValue) )
+			];
+		}
+		else
+		{
+			// create a radio button & label
+			VerticalBox->AddSlot()
+			.Padding(4,4)
+			[
+				// radio button
+				SNew(SHorizontalBox)
+				+SHorizontalBox::Slot()
+				.AutoWidth()
+				.VAlign(VAlign_Top)
+				[
+					SNew(SCheckBox)
+					.Style(FAppStyle::Get(), "RadioButton")
+					.IsChecked_Lambda( [=]() { return (Value == ResultValue) ? ECheckBoxState::Checked : ECheckBoxState::Unchecked; } )
+					.OnCheckStateChanged_Lambda([=](const ECheckBoxState NewState)
+					{
+						if (NewState == ECheckBoxState::Checked)
+						{
+							Value = ResultValue;
+						}
+					})
+
+				]
+
+				// label
+				+SHorizontalBox::Slot()
+				.FillWidth(1.0)
+				.VAlign(VAlign_Top)
+				[
+					SNew(STextBlock)
+					.Text( FText::FromString(OptionValue) )
+				]
+			];
+
+			// update the result value for next time
+			ResultValue++;
+		}
+	};
+}
+
+
+void STurnkeyReadInputIntModal::FinishAction()
+{
+	TSharedPtr<FJsonObject> Result = MakeShared<FJsonObject>();
+	Result->SetNumberField( TEXT("Result"), static_cast<double>(Value) );
+
+	OnFinished.ExecuteIfBound(Result);
+}
+
+
+
+#undef LOCTEXT_NAMESPACE
+#endif //WITH_TURNKEY_EDITOR_IO_SERVER
diff --git a/Engine/Source/Developer/TurnkeyIO/Private/STurnkeyIOWidgets.h b/Engine/Source/Developer/TurnkeyIO/Private/STurnkeyIOWidgets.h
new file mode 100644
--- /dev/null
+++ b/Engine/Source/Developer/TurnkeyIO/Private/STurnkeyIOWidgets.h
@@ -0,0 +1,62 @@
+// Copyright Epic Games, Inc. All Rights Reserved.
+
+#pragma once
+
+#include "CoreMinimal.h"
+#include "TurnkeyEditorIOServer.h"
+
+#if WITH_TURNKEY_EDITOR_IO_SERVER
+#include "Widgets/SCompoundWidget.h"
+
+DECLARE_DELEGATE_OneParam( FOnTurnkeyActionComplete, TSharedPtr<class FJsonObject> );
+
+
+class STurnkeyReadInputModal : public SCompoundWidget
+{
+	SLATE_BEGIN_ARGS( STurnkeyReadInputModal ){}
+	SLATE_ARGUMENT( FString, Prompt );
+	SLATE_ARGUMENT( FString, Default );
+	SLATE_EVENT( FOnTurnkeyActionComplete, OnFinished );
+	SLATE_END_ARGS();
+
+public:
+	void Construct( const FArguments& InArgs );
+
+private:
+	void FinishAction();
+
+	FOnTurnkeyActionComplete OnFinished;
+	FString Prompt;
+	FString Default;
+	FString Value;
+};
+
+
+
+class STurnkeyReadInputIntModal : public SCompoundWidget
+{
+	SLATE_BEGIN_ARGS( STurnkeyReadInputIntModal ){}
+	SLATE_ARGUMENT( FString, Prompt );
+	SLATE_ARGUMENT( TArray<FString>, Options );
+	SLATE_ARGUMENT( bool, IsCancellable );
+	SLATE_ARGUMENT( int, DefaultValue );
+	SLATE_EVENT( FOnTurnkeyActionComplete, OnFinished );
+	SLATE_END_ARGS();
+
+public:
+	void Construct( const FArguments& InArgs );
+
+private:
+	void FinishAction();
+
+	FString Prompt;
+	TArray<FString> Options;
+	bool bIsCancellable;
+	int DefaultValue;
+	int Value;
+	FOnTurnkeyActionComplete OnFinished;
+};
+
+
+#endif //WITH_TURNKEY_EDITOR_IO_SERVER
+
diff --git a/Engine/Source/Developer/TurnkeyIO/Private/TurnkeyEditorIOServer.cpp b/Engine/Source/Developer/TurnkeyIO/Private/TurnkeyEditorIOServer.cpp
new file mode 100644
--- /dev/null
+++ b/Engine/Source/Developer/TurnkeyIO/Private/TurnkeyEditorIOServer.cpp
@@ -0,0 +1,466 @@
+// Copyright Epic Games, Inc. All Rights Reserved.
+
+#include "TurnkeyEditorIOServer.h"
+
+#if WITH_TURNKEY_EDITOR_IO_SERVER
+#include "HAL/PlatformProcess.h"
+#include "HAL/RunnableThread.h"
+#include "HAL/Event.h"
+#include "Misc/MessageDialog.h"
+#include "Async/Async.h"
+#include "Sockets.h"
+#include "SocketSubsystem.h"
+#include "Framework/Docking/TabManager.h"
+#include "Serialization/JsonReader.h"
+#include "Serialization/JsonSerializer.h"
+#include "STurnkeyIOWidgets.h"
+
+#define LOCTEXT_NAMESPACE "TurnkeyEditorIOServer"
+
+DEFINE_LOG_CATEGORY_STATIC(LogTurnkeyIO, Log, All);
+
+
+class FTurnkeyEditorIOServerRunnable : public FRunnable
+{
+public:
+	static const int Version = 1; //keep in sync with .cs
+
+
+	FTurnkeyEditorIOServerRunnable(TSharedPtr<FSocket> InServerSocket)
+		: ServerSocket(InServerSocket)
+	{
+		ReceiveBuffer.AddZeroed(10240);
+		ActionCompleteEvent = FPlatformProcess::GetSynchEventFromPool(true);
+	}
+
+
+	~FTurnkeyEditorIOServerRunnable()
+	{
+		FPlatformProcess::ReturnSynchEventToPool(ActionCompleteEvent);
+		ActionCompleteEvent = nullptr;
+	}
+
+
+	virtual uint32 Run() override
+	{
+		while(!IsEngineExitRequested() && ServerSocket.IsValid())
+		{
+			// wait for a connection
+			UE_LOG( LogTurnkeyIO, Verbose, TEXT("Turnkey IO waiting for connection..."));
+			ClientSocket = TSharedPtr<FSocket>( ServerSocket->Accept(TEXT("TurnkeyIOClient")) );
+			if (!IsClientConnected())
+			{
+				UE_LOG( LogTurnkeyIO, Verbose, TEXT("Client not connected"));
+				continue;
+			}
+			UE_LOG( LogTurnkeyIO, Verbose, TEXT("Client connected"));
+
+			// read the message
+			TSharedPtr<FJsonObject> Message = ReceiveMessage();
+			if (Message.IsValid())
+			{
+				TSharedPtr<FJsonObject> Action = Message->GetObjectField(TEXT("Action"));
+				if (Action.IsValid())
+				{
+					HandleActionMessage(Action);
+				}
+			}
+
+			// clean up for next time
+			UE_LOG( LogTurnkeyIO, Verbose, TEXT("Transaction complete. Disconnecting") );
+			if (ClientSocket)
+			{
+				ClientSocket->Close();
+				ClientSocket.Reset();
+			}
+
+			// wait a short while before reading another message, to give the message log time to display any output from Turnkey
+			if (ServerSocket.IsValid())
+			{
+				FPlatformProcess::Sleep(0.1f);
+			}
+		}
+
+		return 0;
+	}
+
+
+	virtual void Stop() override
+	{
+		if (ServerSocket.IsValid())
+		{
+			ServerSocket->Close();
+			ServerSocket.Reset();
+		}
+		if (ClientSocket.IsValid())
+		{
+			ClientSocket->Close();
+			ClientSocket.Reset();
+		}
+		if (ActionCompleteEvent)
+		{
+			ActionCompleteEvent->Trigger();
+		}
+	}
+
+private:
+
+	void OnActionFinish( TSharedPtr<FJsonObject> Result )
+	{
+		if (Window.IsValid())
+		{
+			TSharedPtr<SWindow> OldWindow = Window;
+			Window.Reset();
+			OldWindow->RequestDestroyWindow();
+		}
+
+		UE_LOG( LogTurnkeyIO, Verbose, TEXT("Turnkey action completed. Sending reply."));
+		SendMessage(Result);
+
+		UE_LOG( LogTurnkeyIO, Verbose, TEXT("Turnkey awaiting ack"));
+		TSharedPtr<FJsonObject> Ack = ReceiveMessage();
+
+		ActionCompleteEvent->Trigger();
+	}
+
+
+	void HandleActionMessage( TSharedPtr<FJsonObject> Action )
+	{
+		// read action type
+		FString Type;
+		if (!Action->TryGetStringField(TEXT("Type"), Type ))
+		{
+			UE_LOG( LogTurnkeyIO, Error, TEXT("Invalid action message") );
+			return;
+		}
+
+		// check the action type
+		ActionCompleteEvent->Reset();
+		if (Type.Equals(TEXT("PauseForUser"), ESearchCase::IgnoreCase ) )
+		{
+			DoPauseForUser(Action);
+		}
+		else if (Type.Equals(TEXT("ReadInput"), ESearchCase::IgnoreCase ) )
+		{
+			DoReadInput(Action);
+		}
+		else if (Type.Equals(TEXT("ReadInputInt"), ESearchCase::IgnoreCase ) )
+		{
+			DoReadInputInt(Action);
+		}
+		else
+		{
+			UE_LOG(LogTurnkeyIO, Error, TEXT("Unknown action type %s"), *Type );
+			return;
+		}
+
+		// wait for the action to finish
+		UE_LOG( LogTurnkeyIO, Verbose, TEXT("Turnkey action in progress. Waiting for completion") );
+		ActionCompleteEvent->Wait();
+	}
+
+
+	void DoPauseForUser( TSharedPtr<FJsonObject> JsonObject )
+	{
+		AsyncTask( ENamedThreads::GameThread, [=]()
+		{
+			// read parameters
+			FString Message;
+			JsonObject->TryGetStringField(TEXT("Message"), Message );
+
+			// display the message box
+			FString Prompt = Message.IsEmpty() ? *LOCTEXT("ReadyToContinue", "Click OK To Continue").ToString() : Message;
+			FText Title = LOCTEXT("Turnkey","Turnkey");
+			FMessageDialog::Open(EAppMsgType::Ok, FText::FromString(Prompt), &Title );
+
+			// send empty response
+			TSharedPtr<FJsonObject> Result = MakeShared<FJsonObject>();
+			OnActionFinish( Result );
+		});
+	}
+
+
+	void DoReadInput( TSharedPtr<FJsonObject> JsonObject )
+	{
+		AsyncTask( ENamedThreads::GameThread, [=]()
+		{
+			// read parameters
+			FString Prompt;
+			JsonObject->TryGetStringField(TEXT("Prompt"), Prompt);
+
+			FString Default;
+			JsonObject->TryGetStringField(TEXT("Default"), Default );
+
+			// show the dialog
+			ShowModal( SNew(STurnkeyReadInputModal)
+				.Prompt(Prompt)
+				.Default(Default)
+				.OnFinished( FOnTurnkeyActionComplete::CreateRaw( this, &FTurnkeyEditorIOServerRunnable::OnActionFinish) )
+			);
+		});
+
+	}
+
+
+	void DoReadInputInt( TSharedPtr<FJsonObject> JsonObject )
+	{
+		AsyncTask( ENamedThreads::GameThread, [=]()
+		{
+			// read parameters
+			FString Prompt;
+			JsonObject->TryGetStringField(TEXT("Prompt"), Prompt);
+
+			TArray<FString> Options;
+			JsonObject->TryGetStringArrayField(TEXT("Options"), Options);
+
+			bool bIsCancellable = true;
+			JsonObject->TryGetBoolField(TEXT("IsCancellable"), bIsCancellable);
+
+			int32 DefaultValue = 0;
+			JsonObject->TryGetNumberField(TEXT("DefaultValue"), DefaultValue);
+
+			// show the dialog
+			ShowModal( SNew(STurnkeyReadInputIntModal)
+				.Prompt(Prompt)
+				.Options(Options)
+				.IsCancellable(bIsCancellable)
+				.DefaultValue(DefaultValue)
+				.OnFinished( FOnTurnkeyActionComplete::CreateRaw( this, &FTurnkeyEditorIOServerRunnable::OnActionFinish) )
+			);
+		});
+	}
+
+
+	void ShowModal( TSharedRef<SWidget> Modal )
+	{
+		// create a window for the widget
+		Window = SNew(SWindow)
+			.Title( LOCTEXT("Turnkey","Turnkey") )
+			.SupportsMaximize(false)
+			.SupportsMinimize(false)
+			.HasCloseButton(false)
+			.AutoCenter(EAutoCenter::PreferredWorkArea)
+			.SizingRule(ESizingRule::Autosized)
+			[
+				Modal
+			];
+
+		// even though HasCloseButton is false, the window can still be closed via Windows' alt-tab screen
+		Window->GetOnWindowClosedEvent().AddLambda( [=]( const TSharedRef<SWindow>& InWindow )
+		{
+			if (Window.IsValid())
+			{
+				Window.Reset();
+		
+				TSharedPtr<FJsonObject> Result = MakeShared<FJsonObject>();
+				Result->SetStringField(TEXT("Error"), TEXT("Turnkey window was closed"));
+				OnActionFinish(Result);
+			}
+		});
+		
+		TSharedPtr<SWindow> RootWindow = FGlobalTabmanager::Get()->GetRootWindow();
+		FSlateApplication::Get().AddModalWindow(Window.ToSharedRef(), RootWindow);
+	}
+
+
+	TSharedPtr<FJsonObject> ReceiveMessage()
+	{
+		UE_LOG(LogTurnkeyIO, Verbose, TEXT("Waiting for message from Turnkey client") );
+
+		// read the message
+		bool bFinished = false;
+		TArray<uint8> Msg;
+		int RecvFails = 0;
+		do
+		{
+			if (!IsClientConnected())
+			{
+				return nullptr;
+			}
+
+			if (ClientSocket->Wait(ESocketWaitConditions::WaitForRead, FTimespan::FromMilliseconds(100)))
+			{
+				int32 BytesRead = 0;
+				if (ClientSocket->Recv(ReceiveBuffer.GetData(), ReceiveBuffer.Num()-1, BytesRead))
+				{
+					Msg.Append( ReceiveBuffer.GetData(), BytesRead );
+				}
+				else
+				{
+					RecvFails++;
+
+					static const int MaxFails = 100; // arbirary, just to make sure it has definititely failed for good
+					if (RecvFails > MaxFails)
+					{
+						UE_LOG(LogTurnkeyIO, Error, TEXT("Client not responding"));
+						return TSharedPtr<FJsonObject>();
+					}
+				}
+			}
+
+		} while( !Msg.Contains('\0') );
+
+		// decode the json object
+		TSharedPtr<FJsonObject> JsonObject;
+		TSharedRef<TJsonReader<>> Reader = TJsonReaderFactory<>::Create(ANSI_TO_TCHAR((const ANSICHAR*)Msg.GetData()));
+		if (!FJsonSerializer::Deserialize(Reader, JsonObject) || !JsonObject.IsValid())
+		{
+			UE_LOG(LogTurnkeyIO, Error, TEXT("Failed to parse message from Turnkey client %s"), *Reader->GetErrorMessage());
+			return TSharedPtr<FJsonObject>();
+		}
+		return JsonObject;
+	}
+
+
+	bool SendMessage( TSharedPtr<FJsonObject> JsonObject )
+	{
+		if (!IsClientConnected())
+		{
+			return false;
+		}
+
+		// add version to message
+		JsonObject->SetNumberField( TEXT("version"), static_cast<double>(Version) );
+
+		// encode the json object and append the message terminator
+		FString JsonString;
+		TSharedRef<TJsonWriter<>> Writer = TJsonWriterFactory<>::Create(&JsonString);
+		if (!FJsonSerializer::Serialize( JsonObject.ToSharedRef(), Writer ))
+		{
+			return false;
+		}
+
+		// send the message
+		auto Msg = StringCast<ANSICHAR>(*JsonString);
+		int32 BytesToSend = Msg.Length()+1; //also send null terminator
+		const ANSICHAR* Data = Msg.Get();
+		while (BytesToSend > 0)
+		{
+			while (!ClientSocket->Wait(ESocketWaitConditions::WaitForWrite, FTimespan::FromSeconds(1.0)))
+			{
+				if (ClientSocket->GetConnectionState() == SCS_ConnectionError)
+				{
+					return false;
+				}
+			}
+
+			int32 BytesSent;
+			if (!ClientSocket->Send( (const uint8*)Data, BytesToSend, BytesSent ))
+			{
+				return false;
+			}
+
+			BytesToSend -= BytesSent;
+			Data += BytesSent;
+		}
+		return true;
+	}
+
+
+	bool IsClientConnected() const
+	{
+		return ClientSocket.IsValid() && ClientSocket->GetConnectionState() == SCS_Connected;
+	}
+
+
+	FEvent* ActionCompleteEvent;
+	TArray<uint8> ReceiveBuffer;
+	TSharedPtr<FSocket> ClientSocket;
+	TSharedPtr<FSocket> ServerSocket;
+	TSharedPtr<SWindow> Window;
+
+};
+
+
+
+
+
+
+
+
+FTurnkeyEditorIOServer::FTurnkeyEditorIOServer()
+	: Port(0)
+	, Runnable(nullptr)
+	, Thread(nullptr)
+{
+	Start();
+}
+
+
+FTurnkeyEditorIOServer::~FTurnkeyEditorIOServer()
+{
+	Stop();
+}
+
+
+FString FTurnkeyEditorIOServer::GetUATParams() const
+{
+	FString Result;
+	if (Thread != nullptr)
+	{
+		Result += FString::Printf(TEXT("-EditorIOPort=%d "), Port);
+	}
+
+	return Result;
+}
+
+
+bool FTurnkeyEditorIOServer::Start()
+{
+	if (Thread == nullptr)
+	{
+		ISocketSubsystem* SocketSubsystem = ISocketSubsystem::Get(PLATFORM_SOCKETSUBSYSTEM);
+		if (SocketSubsystem)
+		{
+			// create an IPv4 TCP server on localhost
+			FName ProtocolType = FNetworkProtocolTypes::IPv4;
+			TSharedPtr<FSocket> NewSocket( SocketSubsystem->CreateSocket(NAME_Stream, TEXT("TurnkeyIOServer"), ProtocolType ) );
+			if (NewSocket.IsValid())
+			{
+				TSharedRef<FInternetAddr> LocalHostAddr = SocketSubsystem->CreateInternetAddr(ProtocolType);
+				LocalHostAddr->SetLoopbackAddress();
+				LocalHostAddr->SetPort(0);
+				if (NewSocket->Bind( *LocalHostAddr ))
+				{
+					if (NewSocket->Listen(1))
+					{
+						Port = NewSocket->GetPortNo();
+						Socket = NewSocket;
+
+						Runnable = new FTurnkeyEditorIOServerRunnable(Socket);
+						Thread = FRunnableThread::Create( Runnable, TEXT("TurnkeyEditorIOServer"), 0, EThreadPriority::TPri_BelowNormal );
+						check( Thread );
+
+						UE_LOG(LogTurnkeyIO, Verbose, TEXT("using port %d"), Port );
+					}
+				}
+			}
+		}
+	}
+
+	return (Thread != nullptr);
+}
+
+
+void FTurnkeyEditorIOServer::Stop()
+{
+	if (Socket.IsValid())
+	{
+		Socket->Close();
+		Socket.Reset();
+	}
+
+	if (Thread)
+	{
+		Thread->Kill();
+		delete Thread;
+		delete Runnable;
+
+		Thread = nullptr;
+		Runnable = nullptr;
+	}
+}
+
+
+#undef LOCTEXT_NAMESPACE
+#endif //WITH_TURNKEY_EDITOR_IO_SERVER
diff --git a/Engine/Source/Developer/TurnkeyIO/Private/TurnkeyEditorIOServer.h b/Engine/Source/Developer/TurnkeyIO/Private/TurnkeyEditorIOServer.h
new file mode 100644
--- /dev/null
+++ b/Engine/Source/Developer/TurnkeyIO/Private/TurnkeyEditorIOServer.h
@@ -0,0 +1,34 @@
+// Copyright Epic Games, Inc. All Rights Reserved.
+
+#pragma once
+
+#include "CoreMinimal.h"
+
+#if !defined(WITH_TURNKEY_EDITOR_IO_SERVER)
+	#error WITH_TURNKEY_EDITOR_IO_SERVER not defined!
+	#define WITH_TURNKEY_EDITOR_IO_SERVER  1 // for IntelliSense
+#endif
+
+#if WITH_TURNKEY_EDITOR_IO_SERVER
+#include "Templates/SharedPointer.h"
+
+class FTurnkeyEditorIOServer : public TSharedFromThis<FTurnkeyEditorIOServer>
+{
+public:
+	FTurnkeyEditorIOServer();
+	~FTurnkeyEditorIOServer();
+
+	FString GetUATParams() const;
+
+private:
+	bool Start();
+	void Stop();
+
+	int Port;
+	class FTurnkeyEditorIOServerRunnable* Runnable;
+	class FRunnableThread* Thread;
+	TSharedPtr<class FSocket> Socket;
+
+};
+#endif //WITH_TURNKEY_EDITOR_IO_SERVER
+
diff --git a/Engine/Source/Developer/TurnkeyIO/Private/TurnkeyIOModule.cpp b/Engine/Source/Developer/TurnkeyIO/Private/TurnkeyIOModule.cpp
new file mode 100644
--- /dev/null
+++ b/Engine/Source/Developer/TurnkeyIO/Private/TurnkeyIOModule.cpp
@@ -0,0 +1,49 @@
+// Copyright Epic Games, Inc. All Rights Reserved.
+
+#include "ITurnkeyIOModule.h"
+#include "TurnkeyEditorIOServer.h"
+
+class FTurnkeyIOModule : public ITurnkeyIOModule
+{
+public:
+	virtual void StartupModule() override
+	{
+#if WITH_TURNKEY_EDITOR_IO_SERVER
+		TurnkeyEditorIOServer = MakeShared<FTurnkeyEditorIOServer>();
+#endif
+	}
+
+	virtual void ShutdownModule() override
+	{
+#if WITH_TURNKEY_EDITOR_IO_SERVER
+		TurnkeyEditorIOServer.Reset();
+#endif
+	}
+
+
+	virtual FString GetUATParams() const override
+	{
+		FString Result = TEXT("-EditorIO ");
+
+#if WITH_TURNKEY_EDITOR_IO_SERVER
+		if (TurnkeyEditorIOServer.IsValid())
+		{
+			Result += TurnkeyEditorIOServer->GetUATParams();
+		}
+#endif
+
+		return Result;
+	}
+
+private:
+
+#if WITH_TURNKEY_EDITOR_IO_SERVER
+	TSharedPtr<FTurnkeyEditorIOServer> TurnkeyEditorIOServer;
+#endif
+
+};
+
+
+IMPLEMENT_MODULE(FTurnkeyIOModule, TurnkeyIO);
+
+#undef LOCTEXT_NAMESPACE
diff --git a/Engine/Source/Developer/TurnkeyIO/Public/ITurnkeyIOModule.h b/Engine/Source/Developer/TurnkeyIO/Public/ITurnkeyIOModule.h
new file mode 100644
--- /dev/null
+++ b/Engine/Source/Developer/TurnkeyIO/Public/ITurnkeyIOModule.h
@@ -0,0 +1,27 @@
+// Copyright Epic Games, Inc. All Rights Reserved.
+
+#pragma once
+
+#include "Modules/ModuleInterface.h"
+#include "Modules/ModuleManager.h"
+
+/**
+ * Interface for Turnkey IO module, used for when Turnkey requires user input
+ */
+class ITurnkeyIOModule : public IModuleInterface
+{
+public:
+	/** 
+	 * Returns the command line arguments for passing to UAT
+	 */
+	virtual FString GetUATParams() const = 0;
+
+	/**
+	 * Returns a reference to the Turnkey IO module instance
+	 */
+	static ITurnkeyIOModule& Get()
+	{
+		static const FName TurnkeyIOModuleName = "TurnkeyIO";
+		return FModuleManager::LoadModuleChecked<ITurnkeyIOModule>(TurnkeyIOModuleName);
+	}
+};
diff --git a/Engine/Source/Developer/TurnkeyIO/TurnkeyIO.Build.cs b/Engine/Source/Developer/TurnkeyIO/TurnkeyIO.Build.cs
new file mode 100644
--- /dev/null
+++ b/Engine/Source/Developer/TurnkeyIO/TurnkeyIO.Build.cs
@@ -0,0 +1,34 @@
+// Copyright Epic Games, Inc. All Rights Reserved.
+
+using UnrealBuildTool;
+
+public class TurnkeyIO : ModuleRules
+{
+	public TurnkeyIO(ReadOnlyTargetRules Target) : base(Target)
+	{
+		PrivateDependencyModuleNames.AddRange(
+			new string[] {
+				"Core",
+			}
+		);
+
+		if (Target.bBuildDeveloperTools)
+		{
+			PrivateDependencyModuleNames.AddRange(
+				new string[] {
+				"CoreUObject",
+				"Sockets",
+				"Slate",
+				"SlateCore",
+				"Json",
+				}
+			);
+
+			PublicDefinitions.Add("WITH_TURNKEY_EDITOR_IO_SERVER=1");
+		}
+		else
+		{
+			PublicDefinitions.Add("WITH_TURNKEY_EDITOR_IO_SERVER=0");
+		}
+	}
+}
diff --git a/Engine/Source/Editor/TurnkeySupport/Private/TurnkeyEditorSupport.cpp b/Engine/Source/Editor/TurnkeySupport/Private/TurnkeyEditorSupport.cpp
--- a/Engine/Source/Editor/TurnkeySupport/Private/TurnkeyEditorSupport.cpp
+++ b/Engine/Source/Editor/TurnkeySupport/Private/TurnkeyEditorSupport.cpp
@@ -1,6 +1,7 @@
 // Copyright Epic Games, Inc. All Rights Reserved.
 
 #include "TurnkeyEditorSupport.h"
+#include "ITurnkeyIOModule.h"
 #include "Misc/AssertionMacros.h"
 
 #if WITH_EDITOR
@@ -130,7 +131,7 @@ void FTurnkeyEditorSupport::LaunchRunningMap(const FString& DeviceId, const FStr
 			// if we want to check device flash before we start cooking, kick it off now. we could delay this 
 			if (bUseTurnkey)
 			{
-				FString CommandLine = FString::Printf(TEXT("Turnkey -command=VerifySdk -UpdateIfNeeded -platform=%s -EditorIO -noturnkeyvariables -device=%s -utf8output -WaitForUATMutex %s"), *UBTPlatformName, *TargetDeviceId.GetDeviceName(), *PlatformInfo->UATCommandLine);
+				FString CommandLine = FString::Printf(TEXT("Turnkey -command=VerifySdk -UpdateIfNeeded -platform=%s -noturnkeyvariables -device=%s -utf8output -WaitForUATMutex %s %s"), *UBTPlatformName, *TargetDeviceId.GetDeviceName(), *PlatformInfo->UATCommandLine, *ITurnkeyIOModule::Get().GetUATParams());
 				if (!ProjectPath.IsEmpty())
 				{
 					CommandLine = FString::Printf(TEXT(" -ScriptsForProject=\"%s\" %s -project=\"%s\""), *ProjectPath, *CommandLine, *ProjectPath);
diff --git a/Engine/Source/Editor/TurnkeySupport/Private/TurnkeySupportModule.cpp b/Engine/Source/Editor/TurnkeySupport/Private/TurnkeySupportModule.cpp
--- a/Engine/Source/Editor/TurnkeySupport/Private/TurnkeySupportModule.cpp
+++ b/Engine/Source/Editor/TurnkeySupport/Private/TurnkeySupportModule.cpp
@@ -43,6 +43,7 @@
 #include "UObject/UObjectIterator.h"
 #include "ToolMenus.h"
 #include "TurnkeyEditorSupport.h"
+#include "ITurnkeyIOModule.h"
 
 #include "Misc/App.h"
 #include "Framework/Application/SlateApplication.h"
@@ -502,7 +503,7 @@ public:
 		}
 
 
-		FString TurnkeyParams = FString::Printf(TEXT("-command=VerifySdk -platform=%s -UpdateIfNeeded -EditorIO"), *UBTPlatformString);
+		FString TurnkeyParams = FString::Printf(TEXT("-command=VerifySdk -platform=%s -UpdateIfNeeded %s"), *UBTPlatformString, *ITurnkeyIOModule::Get().GetUATParams());
 		if (!ProjectPath.IsEmpty())
 		{
 			TurnkeyParams.Appendf(TEXT(" -project=\"%s\""), *ProjectPath);
@@ -723,7 +724,7 @@ static void TurnkeyInstallSdk(FString PlatformName, bool bPreferFull, bool bForc
 	{
 		CommandLine.Appendf(TEXT("-ScriptsForProject=\"%s\" "), *ProjectPath);
 	}
-	CommandLine.Appendf(TEXT("Turnkey -command=VerifySdk -UpdateIfNeeded -platform=%s %s -EditorIO -noturnkeyvariables -utf8output -WaitForUATMutex"), *PlatformName, *OptionalOptions);
+	CommandLine.Appendf(TEXT("Turnkey -command=VerifySdk -UpdateIfNeeded -platform=%s %s %s -noturnkeyvariables -utf8output -WaitForUATMutex"), *PlatformName, *OptionalOptions, *ITurnkeyIOModule::Get().GetUATParams());
 
 	FText TaskName = LOCTEXT("InstallingSdk", "Installing Sdk");
 	FTurnkeyEditorSupport::RunUAT(CommandLine, FText::FromString(PlatformName), TaskName, TaskName, FEditorStyle::GetBrush(TEXT("MainFrame.PackageProject")),
diff --git a/Engine/Source/Editor/TurnkeySupport/TurnkeySupport.Build.cs b/Engine/Source/Editor/TurnkeySupport/TurnkeySupport.Build.cs
--- a/Engine/Source/Editor/TurnkeySupport/TurnkeySupport.Build.cs
+++ b/Engine/Source/Editor/TurnkeySupport/TurnkeySupport.Build.cs
@@ -28,6 +28,7 @@ public class TurnkeySupport : ModuleRules
  				"LauncherServices",
 				"SourceControl",
 				"EditorStyle",
+				"TurnkeyIO",
 			}
 		);
 
diff --git a/Engine/Source/Programs/AutomationTool/Turnkey/Providers/TurnkeyEditorIOClient.cs b/Engine/Source/Programs/AutomationTool/Turnkey/Providers/TurnkeyEditorIOClient.cs
new file mode 100644
--- /dev/null
+++ b/Engine/Source/Programs/AutomationTool/Turnkey/Providers/TurnkeyEditorIOClient.cs
@@ -0,0 +1,290 @@
+// Copyright Epic Games, Inc. All Rights Reserved.
+
+using System;
+using System.Collections.Generic;
+using System.Net.Sockets;
+using System.Net;
+using System.Text;
+using EpicGames.Core;
+using System.Linq;
+using System.IO;
+
+namespace Turnkey
+{
+	class EditorIOClient : IOProvider	
+	{
+		readonly int Version = 1; //keep in sync with .cpp
+
+		int Port;
+		Socket IOSocket;
+		Lazy<IOProvider> FallbackIOProvider;
+		bool bVerbose = false;
+
+		public EditorIOClient( int InPort )
+		{
+			Port = InPort;
+			FallbackIOProvider = new Lazy<IOProvider>(() => { return new HybridIOProvider(); } );
+		}
+
+		~EditorIOClient()
+		{
+			Disconnect();
+		}
+
+		private void VerboseMessage( string Message )
+		{
+			if(bVerbose)
+			{
+				Console.WriteLine(Message);
+			}
+		}
+
+		public override void Log(string Message, bool bAppendNewLine)
+		{
+			// no socket needed for this - editor will be logging stdout anyway
+			if (bAppendNewLine)
+			{
+				Console.WriteLine(Message);
+			}
+			else
+			{
+				Console.Write(Message);
+			}
+		}
+
+		public override void PauseForUser(string Message, bool bAppendNewLine)
+		{
+			if (Connect())
+			{
+				// send the action definition
+				if (SendAction( "PauseForUser", (Json) =>
+				{
+					Json.WriteValue("Message", Message);
+				}))
+				{					
+					// wait for any response
+					JsonObject ReceivedMessage = ReceiveMessage();
+					Disconnect();
+
+					// make sure there was a response, otherwise use the fallback
+					if (ReceivedMessage != null)
+					{
+						return;
+					}
+				}
+			}
+
+			// failed to connect - use the fallback instead
+			FallbackIOProvider.Value.PauseForUser(Message, bAppendNewLine);
+		}
+
+		public override string ReadInput(string Prompt, string Default, bool bAppendNewLine)
+		{
+			if (Connect())
+			{
+				// send the action message
+				if (SendAction( "ReadInput", (Json) =>
+				{
+					Json.WriteValue("Prompt",  Prompt);
+					Json.WriteValue("Default", Default);
+				}))
+				{
+					// wait for the response
+					JsonObject ReceivedMessage = ReceiveMessage();
+					Disconnect();
+
+					// parse the result
+					if (ReceivedMessage != null)
+					{
+						string Result;
+						if (ReceivedMessage.TryGetStringField("Result", out Result))
+						{
+							return Result;
+						}
+						return Default;
+					}
+				}
+			}
+
+			// failed to connect - use the fallback instead
+			return FallbackIOProvider.Value.ReadInput(Prompt, Default, bAppendNewLine);
+		}
+
+		public override int ReadInputInt(string Prompt, List<string> Options, bool bIsCancellable, int DefaultValue, bool bAppendNewLine)
+		{
+			if (Connect())
+			{
+				// send the action message
+				if (SendAction( "ReadInputInt", (Json) =>
+				{
+					Json.WriteValue("Prompt",             Prompt);
+					Json.WriteValue("IsCancellable",      bIsCancellable);
+					Json.WriteValue("DefaultValue",       DefaultValue);
+					Json.WriteStringArrayField("Options", Options);
+				}))
+				{
+					// wait for response
+					JsonObject ReceivedMessage = ReceiveMessage();
+					Disconnect();
+
+					// parse the response
+					if (ReceivedMessage != null)
+					{
+						int Choice;
+						if (ReceivedMessage.TryGetIntegerField("Result", out Choice) == false || Choice < 0 || Choice >= Options.Count + (bIsCancellable ? 1 : 0))
+						{
+							TurnkeyUtils.Log("Invalid choice");
+							return DefaultValue;
+						}
+						else
+						{
+							return Choice;
+						}
+					}
+				}
+			}
+
+			// failed to connect - use the fallback instead
+			return FallbackIOProvider.Value.ReadInputInt(Prompt, Options, bIsCancellable, DefaultValue, bAppendNewLine);
+		}
+
+
+		#region socket functions
+
+		bool Connect()
+		{
+			try
+			{
+				// create a IPv4 blocking TCP socket & connect to the given port on localhost
+				IOSocket = new Socket(AddressFamily.InterNetwork, SocketType.Stream, ProtocolType.Tcp );
+				IOSocket.Blocking = true;
+
+				IPEndPoint EndPoint = new IPEndPoint(IPAddress.Loopback, Port );
+				IOSocket.Connect(EndPoint);
+				return true;
+			}
+			catch( Exception e )
+			{
+				Console.WriteLine($"Error: Failed to connect to Turnkey IO server on port {Port}. {e}");
+			}
+			return false;
+		}
+
+		void Disconnect()
+		{
+			if (IOSocket != null)
+			{
+				IOSocket.Disconnect(false);
+				IOSocket.Shutdown(SocketShutdown.Both);
+				IOSocket.Close();
+			}
+		}
+
+		protected JsonObject ReceiveMessage()
+		{
+			VerboseMessage($"TurnkeySocketIO waiting for server response");
+
+			JsonObject Result = null;
+
+			// read the message
+			List<byte> Msg = new List<byte>();
+			try
+			{
+				byte[] ReceiveBuffer = new byte[10240];
+				int MessageTerminator = -1;
+				do
+				{
+					int BytesReceived = IOSocket.Receive(ReceiveBuffer);
+					if (BytesReceived > 0)
+					{
+						Msg.AddRange(ReceiveBuffer.Take(BytesReceived));
+						MessageTerminator = Msg.IndexOf(0);
+					}
+				} while(MessageTerminator < 0);
+
+				//remove the terminator. not expecting more than one message so this isn't handled
+				Msg.RemoveRange( MessageTerminator, Msg.Count - MessageTerminator );
+			}
+			catch( Exception e )
+			{
+				Console.WriteLine($"Error: TurnkeySocketIO read failed {e}");
+				return null;
+			}
+			SendAck();
+
+			// parse the json
+			try
+			{
+				string JsonString = Encoding.ASCII.GetString(Msg.ToArray());
+				Result = JsonObject.Parse(JsonString);
+			}
+			catch( Exception e )
+			{
+				Console.WriteLine($"Error: TurnkeySocketIO parse failed {e}");
+				return null;
+			}
+
+			// check for error message
+			string ErrorMessage;
+			if (Result.TryGetStringField("Error", out ErrorMessage) )
+			{
+				Console.WriteLine($"Error: TurnkeySocketIO server returned {ErrorMessage}");
+				return null;
+			}
+
+			// all done
+			return Result;
+		}
+
+
+
+		protected bool SendMessage( Action<JsonWriter> MessageFunc )
+		{
+			// build message
+			StringWriter JsonMessage = new StringWriter();
+			JsonWriter Json = new JsonWriter(JsonMessage, Style:JsonWriterStyle.Readable );
+			Json.WriteObjectStart();
+				Json.WriteValue("version", Version );
+				MessageFunc(Json);
+			Json.WriteObjectEnd();
+
+			// send the message with a null terminator
+			string Message = JsonMessage.ToString() + '\0';
+			byte[] MessageBytes = Encoding.ASCII.GetBytes(Message);
+			try
+			{
+				int BytesSent = IOSocket.Send(MessageBytes, SocketFlags.None);
+			}
+			catch( Exception e )
+			{
+				Console.WriteLine($"Error: TurnkeySocketIO write failed {e}");
+				return false;
+			}
+
+			return true;
+		}
+
+		protected bool SendAck()
+		{
+			return SendMessage( (Json) =>
+			{
+				Json.WriteValue("Reply", "done");
+			});
+		}
+
+		protected bool SendAction( string Type, Action<JsonWriter> ActionFunc )
+		{
+			VerboseMessage($"TurnkeySocketIO sending action {Type}");
+
+			return SendMessage( (Json) =>
+			{
+				Json.WriteObjectStart("Action");
+					Json.WriteValue("Type", Type);
+					ActionFunc(Json);
+				Json.WriteObjectEnd();
+			});
+		}
+
+		#endregion
+	}
+}
diff --git a/Engine/Source/Programs/AutomationTool/Turnkey/Turnkey.Automation.cs b/Engine/Source/Programs/AutomationTool/Turnkey/Turnkey.Automation.cs
--- a/Engine/Source/Programs/AutomationTool/Turnkey/Turnkey.Automation.cs
+++ b/Engine/Source/Programs/AutomationTool/Turnkey/Turnkey.Automation.cs
@@ -14,7 +14,12 @@ namespace Turnkey
 		public override ExitCode Execute()
 		{
 			IOProvider IOProvider;
-			if (ParseParam("EditorIO"))
+			if (ParseParam("EditorIOPort"))
+			{
+				int Port = ParseParamInt("EditorIOPort");
+				IOProvider = new EditorIOClient(Port);
+			}
+			else if (ParseParam("EditorIO"))
 			{
 				IOProvider = new HybridIOProvider();
 			}
