can select a radio option by double clicking the label in the Turnkey IO widget.

#rnx

diff --git a/Engine/Source/Developer/TurnkeyIO/Private/STurnkeyIOWidgets.cpp b/Engine/Source/Developer/TurnkeyIO/Private/STurnkeyIOWidgets.cpp
--- a/Engine/Source/Developer/TurnkeyIO/Private/STurnkeyIOWidgets.cpp
+++ b/Engine/Source/Developer/TurnkeyIO/Private/STurnkeyIOWidgets.cpp
@@ -207,6 +207,8 @@ void STurnkeyReadInputIntModal::Construct( const FArguments& InArgs )
 		}
 		else
 		{
+			TSharedPtr<STextBlock> TextBlock;
+
 			// create a radio button & label
 			VerticalBox->AddSlot()
 			.Padding(4,4)
@@ -235,11 +237,27 @@ void STurnkeyReadInputIntModal::Construct( const FArguments& InArgs )
 				.FillWidth(1.0)
 				.VAlign(VAlign_Top)
 				[
-					SNew(STextBlock)
+					SAssignNew(TextBlock, STextBlock)
 					.Text( FText::FromString(OptionValue) )
+					.OnDoubleClicked_Lambda([=](const FGeometry& MyGeometry, const FPointerEvent& PointerEvent)
+					{
+						// double click the label to select the option immediately
+						Value = ResultValue;
+						FinishAction();
+						return FReply::Handled();
+					})
+
 				]
 			];
 
+			// change the option if the user clicks on the label
+			TextBlock->SetOnMouseButtonUp( FPointerEventHandler::CreateLambda( [=](const FGeometry& MyGeometry, const FPointerEvent& PointerEvent)
+			{
+				Value = ResultValue;
+				return FReply::Handled();
+			}));
+
+
 			// update the result value for next time
 			ResultValue++;
 		}
