non-unity fix

#rnx

diff --git a/Engine/Source/Developer/TurnkeyIO/Private/STurnkeyIOWidgets.cpp b/Engine/Source/Developer/TurnkeyIO/Private/STurnkeyIOWidgets.cpp
--- a/Engine/Source/Developer/TurnkeyIO/Private/STurnkeyIOWidgets.cpp
+++ b/Engine/Source/Developer/TurnkeyIO/Private/STurnkeyIOWidgets.cpp
@@ -7,6 +7,8 @@
 #include "Widgets/Input/SButton.h"
 #include "Widgets/Input/SCheckBox.h"
 #include "Widgets/Input/SEditableTextBox.h"
+#include "Widgets/SBoxPanel.h"
+#include "Framework/Application/SlateApplication.h"
 #include "Dom/JsonObject.h"
 
 #define LOCTEXT_NAMESPACE "TurnkeySlateIO"
