Added GetUserConfirmation to Turnkey IO provider, allowing better Yes/No prompts.

#rnx

diff --git a/Engine/Source/Developer/TurnkeyIO/Private/TurnkeyEditorIOServer.cpp b/Engine/Source/Developer/TurnkeyIO/Private/TurnkeyEditorIOServer.cpp
--- a/Engine/Source/Developer/TurnkeyIO/Private/TurnkeyEditorIOServer.cpp
+++ b/Engine/Source/Developer/TurnkeyIO/Private/TurnkeyEditorIOServer.cpp
@@ -148,6 +148,10 @@ private:
 		{
 			DoReadInputInt(Action);
 		}
+		else if (Type.Equals(TEXT("GetConfirmation"), ESearchCase::IgnoreCase ) )
+		{
+			DoGetUserConfirmation(Action);
+		}
 		else
 		{
 			UE_LOG(LogTurnkeyIO, Error, TEXT("Unknown action type %s"), *Type );
@@ -231,6 +235,30 @@ private:
 	}
 
 
+	void DoGetUserConfirmation( TSharedPtr<FJsonObject> JsonObject )
+	{
+		AsyncTask( ENamedThreads::GameThread, [=]()
+		{
+			// read parameters
+			FString Message;
+			JsonObject->TryGetStringField(TEXT("Message"), Message );
+
+			bool bDefaultValue = true;
+			JsonObject->TryGetBoolField(TEXT("DefaultValue"), bDefaultValue );
+
+			// display the message box
+			FText Title = LOCTEXT("Turnkey","Turnkey");
+			EAppReturnType::Type Result = (bDefaultValue ? EAppReturnType::Yes : EAppReturnType::No );
+			Result = FMessageDialog::Open(EAppMsgType::YesNo, Result, FText::FromString(Message), &Title );
+
+			// send empty response
+			TSharedPtr<FJsonObject> JsonResult = MakeShared<FJsonObject>();
+			JsonResult->SetBoolField(TEXT("Result"), (Result == EAppReturnType::Yes) );
+			OnActionFinish( JsonResult );
+		});
+	}
+
+
 	void ShowModal( TSharedRef<SWidget> Modal )
 	{
 		// create a window for the widget
diff --git a/Engine/Source/Programs/AutomationTool/Turnkey/CopyAndRun.cs b/Engine/Source/Programs/AutomationTool/Turnkey/CopyAndRun.cs
--- a/Engine/Source/Programs/AutomationTool/Turnkey/CopyAndRun.cs
+++ b/Engine/Source/Programs/AutomationTool/Turnkey/CopyAndRun.cs
@@ -138,8 +138,8 @@ namespace Turnkey
 
 					if (!bUnattended)
 					{
-						string Response = TurnkeyUtils.ReadInput("Do you want to attempt again? [y/N]", "N");
-						if (string.Compare(Response, "Y", true) != 0)
+						bool bResponse = TurnkeyUtils.GetUserConfirmation("Do you want to attempt again?", false);
+						if (bResponse == false)
 						{
 							bDone = true;
 						}
diff --git a/Engine/Source/Programs/AutomationTool/Turnkey/FileSource.cs b/Engine/Source/Programs/AutomationTool/Turnkey/FileSource.cs
--- a/Engine/Source/Programs/AutomationTool/Turnkey/FileSource.cs
+++ b/Engine/Source/Programs/AutomationTool/Turnkey/FileSource.cs
@@ -522,8 +522,8 @@ namespace Turnkey
 		// 				if (!bUnattended)
 		// 				{
 		// 					// @todo turnkey - have studio settings 
-		// 					string Response = TurnkeyUtils.ReadInput("AutoSdks are not setup, but your studio has support. Would you like to set it up now? [Y/n]", "Y");
-		// 					if (string.Compare(Response, "Y", true) == 0)
+		// 					bool bResponse = TurnkeyUtils.GetUserConfirmation("AutoSdks are not setup, but your studio has support. Would you like to set it up now?", true);
+		// 					if (bResponse)
 		// 					{
 		// 						bAttemptAutoSdkSetup = true;
 		// 						bSetupEnvVarAfterInstall = true;
diff --git a/Engine/Source/Programs/AutomationTool/Turnkey/IOProvider.cs b/Engine/Source/Programs/AutomationTool/Turnkey/IOProvider.cs
--- a/Engine/Source/Programs/AutomationTool/Turnkey/IOProvider.cs
+++ b/Engine/Source/Programs/AutomationTool/Turnkey/IOProvider.cs
@@ -22,5 +22,7 @@ namespace Turnkey
 		public abstract string ReadInput(string Prompt, string DefaultValue, bool bAppendNewLine);
 		
 		public abstract int ReadInputInt(string Prompt, List<string> Options, bool bIsCancellable, int DefaultValue, bool bAppendNewLine);
+
+		public abstract bool GetUserConfirmation(string Message, bool bDefaultValue, bool bAppendNewLine);
 	}
 }
diff --git a/Engine/Source/Programs/AutomationTool/Turnkey/Providers/ConsoleIOProvider.cs b/Engine/Source/Programs/AutomationTool/Turnkey/Providers/ConsoleIOProvider.cs
--- a/Engine/Source/Programs/AutomationTool/Turnkey/Providers/ConsoleIOProvider.cs
+++ b/Engine/Source/Programs/AutomationTool/Turnkey/Providers/ConsoleIOProvider.cs
@@ -98,5 +98,19 @@ namespace Turnkey
 				}
 			}
 		}
+
+		public override bool GetUserConfirmation(string Message, bool bDefaultValue, bool bAppendNewLine)
+		{
+			string Prompt = string.Format($"{Message} [{(bDefaultValue ? 'Y' : 'y')}/{(bDefaultValue ? 'n' : 'N')}]");
+			string Result = ReadInput(Prompt, bDefaultValue ? "Y" : "N", bAppendNewLine);
+			if (string.IsNullOrEmpty(Result))
+			{
+				return bDefaultValue;
+			}
+			else
+			{
+				return string.Equals( Result.Trim(), "Y", StringComparison.OrdinalIgnoreCase );
+			}
+		}
 	}
 }
diff --git a/Engine/Source/Programs/AutomationTool/Turnkey/Providers/PerforceCopyProvider.cs b/Engine/Source/Programs/AutomationTool/Turnkey/Providers/PerforceCopyProvider.cs
--- a/Engine/Source/Programs/AutomationTool/Turnkey/Providers/PerforceCopyProvider.cs
+++ b/Engine/Source/Programs/AutomationTool/Turnkey/Providers/PerforceCopyProvider.cs
@@ -228,8 +228,8 @@ namespace Turnkey
 				{
 					TurnkeyUtils.Log("Unable to find a clientspec for the perforce operation {0}, looking for a depot client", Operation);
 
-					string Response = TurnkeyUtils.ReadInput("Would you like to create one? [y/N]", "N");
-					if (string.Compare(Response, "Y", true) != 0)
+					bool bResponse = TurnkeyUtils.GetUserConfirmation("Would you like to create one?", false);
+					if (bResponse == false)
 					{
 						// make sure to try again next time
 						PerforceConnection = null;
diff --git a/Engine/Source/Programs/AutomationTool/Turnkey/Providers/TurnkeyEditorIOClient.cs b/Engine/Source/Programs/AutomationTool/Turnkey/Providers/TurnkeyEditorIOClient.cs
--- a/Engine/Source/Programs/AutomationTool/Turnkey/Providers/TurnkeyEditorIOClient.cs
+++ b/Engine/Source/Programs/AutomationTool/Turnkey/Providers/TurnkeyEditorIOClient.cs
@@ -148,6 +148,39 @@ namespace Turnkey
 			return FallbackIOProvider.Value.ReadInputInt(Prompt, Options, bIsCancellable, DefaultValue, bAppendNewLine);
 		}
 
+		public override bool GetUserConfirmation(string Message, bool bDefaultValue, bool bAppendNewLine)
+		{
+			if (Connect())
+			{
+				// send the action message
+				if (SendAction( "GetConfirmation", (Json) =>
+				{
+					Json.WriteValue("Message",      Message);
+					Json.WriteValue("DefaultValue", bDefaultValue);
+				}))
+				{
+					// wait for the response
+					JsonObject ReceivedMessage = ReceiveMessage();
+					Disconnect();
+
+					// parse the result
+					if (ReceivedMessage != null)
+					{
+						bool bResult;
+						if (ReceivedMessage.TryGetBoolField("Result", out bResult))
+						{
+							return bResult;
+						}
+						return bDefaultValue;
+					}
+				}
+			}
+
+			// failed to connect - use the fallback instead
+			return FallbackIOProvider.Value.GetUserConfirmation(Message, bDefaultValue, bAppendNewLine);
+
+		}
+
 
 		#region socket functions
 
diff --git a/Engine/Source/Programs/AutomationTool/Turnkey/SdkInfo.cs b/Engine/Source/Programs/AutomationTool/Turnkey/SdkInfo.cs
--- a/Engine/Source/Programs/AutomationTool/Turnkey/SdkInfo.cs
+++ b/Engine/Source/Programs/AutomationTool/Turnkey/SdkInfo.cs
@@ -151,8 +151,8 @@ namespace Turnkey
 					return false;
 				}
 
-				string Response = TurnkeyUtils.ReadInput("The AutoSdk system is not setup on this machine. Would you like to set it up now? [Y/n]", "Y");
-				if (string.Compare(Response, "Y", true) == 0)
+				bool bResponse = TurnkeyUtils.GetUserConfirmation("The AutoSdk system is not setup on this machine. Would you like to set it up now?", true);
+				if (bResponse)
 				{
 					bSetupEnvVarAfterInstall = true;
 				}
diff --git a/Engine/Source/Programs/AutomationTool/Turnkey/TurnkeyUtils.cs b/Engine/Source/Programs/AutomationTool/Turnkey/TurnkeyUtils.cs
--- a/Engine/Source/Programs/AutomationTool/Turnkey/TurnkeyUtils.cs
+++ b/Engine/Source/Programs/AutomationTool/Turnkey/TurnkeyUtils.cs
@@ -338,8 +338,8 @@ namespace Turnkey
 						string ProjectChoice = TurnkeyUtils.ReadInput("Enter path to .uproject file:");
 						if (!File.Exists(ProjectChoice))
 						{
-							string Response = TurnkeyUtils.ReadInput(string.Format("'{0}' doesn't exist. Would you like to enter another path? [y/N]", ProjectChoice), "N");
-							if (Response.ToLower() != "y")
+							bool bResponse = TurnkeyUtils.GetUserConfirmation(string.Format("'{0}' doesn't exist. Would you like to enter another path?", ProjectChoice), false);
+							if (bResponse == false)
 							{
 								break;
 							}
@@ -823,6 +823,11 @@ namespace Turnkey
 		{
 			return IOProvider.ReadInputInt(Prompt, Options, bIsCancellable, DefaultValue, bAppendNewLine: true);
 		}
+
+		public static bool GetUserConfirmation(string Message, bool bDefaultValue)
+		{
+			return IOProvider.GetUserConfirmation(Message, bDefaultValue, true );
+		}
 #endregion
 
 #region Temp files [move to LocalCache]
