non unity fix

#rnx

diff --git a/Engine/Source/Developer/TurnkeyIO/Private/TurnkeyEditorIOServer.cpp b/Engine/Source/Developer/TurnkeyIO/Private/TurnkeyEditorIOServer.cpp
--- a/Engine/Source/Developer/TurnkeyIO/Private/TurnkeyEditorIOServer.cpp
+++ b/Engine/Source/Developer/TurnkeyIO/Private/TurnkeyEditorIOServer.cpp
@@ -11,6 +11,7 @@
 #include "Sockets.h"
 #include "SocketSubsystem.h"
 #include "Framework/Docking/TabManager.h"
+#include "Framework/Application/SlateApplication.h"
 #include "Serialization/JsonReader.h"
 #include "Serialization/JsonSerializer.h"
 #include "STurnkeyIOWidgets.h"
